[Mutable] Disable MutableExtensionExampleModule Operation registration as a workaround for shutdown crashes in localization commandlet.

#rnx

diff --git a/Engine/Plugins/Experimental/MutableExtensionExample/Source/MutableExtensionExample/Private/MutableExtensionExampleModule.cpp b/Engine/Plugins/Experimental/MutableExtensionExample/Source/MutableExtensionExample/Private/MutableExtensionExampleModule.cpp
--- a/Engine/Plugins/Experimental/MutableExtensionExample/Source/MutableExtensionExample/Private/MutableExtensionExampleModule.cpp
+++ b/Engine/Plugins/Experimental/MutableExtensionExample/Source/MutableExtensionExample/Private/MutableExtensionExampleModule.cpp
@@ -12,17 +12,23 @@ IMPLEMENT_MODULE(FMutableExtensionExampleModule, MutableExtensionExample);
 
 void FMutableExtensionExampleModule::StartupModule()
 {
-	ICustomizableObjectModule& Module = ICustomizableObjectModule::Get();
+	// TODO: Disable extension registration as a workaround to unblock Localization build crashes. 
+	// This should only make the operations unavailable.
 
-	Module.RegisterExternalOperation(FMeshIdentity::StaticStruct());
-	Module.RegisterExternalOperation(FMeshNoise::StaticStruct());
+	//ICustomizableObjectModule& Module = ICustomizableObjectModule::Get();
+
+	//Module.RegisterExternalOperation(FMeshIdentity::StaticStruct());
+	//Module.RegisterExternalOperation(FMeshNoise::StaticStruct());
 }
 
 
 void FMutableExtensionExampleModule::ShutdownModule()
 {
-	ICustomizableObjectModule& Module = ICustomizableObjectModule::Get();
+	// TODO: Disable extension registration as a workaround to unblock Localization build crashes. 
+	// StaticStruct() result seems to be GC'd at this point, more investigation is needed.
+	
+	//ICustomizableObjectModule& Module = ICustomizableObjectModule::Get();
 
-	Module.UnregisterExternalOperation(FMeshIdentity::StaticStruct());
-	Module.UnregisterExternalOperation(FMeshNoise::StaticStruct());
+	//Module.UnregisterExternalOperation(FMeshIdentity::StaticStruct());
+	//Module.UnregisterExternalOperation(FMeshNoise::StaticStruct());
 }
