[Mutable] Fix multiple issues with External Operations.

#rnx

#FailedValidations Preflight

diff --git a/Engine/Plugins/Experimental/MutableExtensionExample/Source/MutableExtensionExample/Private/MutableExtensionExampleModule.cpp b/Engine/Plugins/Experimental/MutableExtensionExample/Source/MutableExtensionExample/Private/MutableExtensionExampleModule.cpp
--- a/Engine/Plugins/Experimental/MutableExtensionExample/Source/MutableExtensionExample/Private/MutableExtensionExampleModule.cpp
+++ b/Engine/Plugins/Experimental/MutableExtensionExample/Source/MutableExtensionExample/Private/MutableExtensionExampleModule.cpp
@@ -12,23 +12,17 @@ IMPLEMENT_MODULE(FMutableExtensionExampleModule, MutableExtensionExample);
 
 void FMutableExtensionExampleModule::StartupModule()
 {
-	// TODO: Disable extension registration as a workaround to unblock Localization build crashes. 
-	// This should only make the operations unavailable.
+	ICustomizableObjectModule& Module = ICustomizableObjectModule::Get();
 
-	//ICustomizableObjectModule& Module = ICustomizableObjectModule::Get();
-
-	//Module.RegisterExternalOperation(FMeshIdentity::StaticStruct());
-	//Module.RegisterExternalOperation(FMeshNoise::StaticStruct());
+	Module.RegisterExternalOperation(FMeshIdentity::StaticStruct());
+	Module.RegisterExternalOperation(FMeshNoise::StaticStruct());
 }
 
 
 void FMutableExtensionExampleModule::ShutdownModule()
 {
-	// TODO: Disable extension registration as a workaround to unblock Localization build crashes. 
-	// StaticStruct() result seems to be GC'd at this point, more investigation is needed.
-	
-	//ICustomizableObjectModule& Module = ICustomizableObjectModule::Get();
+	ICustomizableObjectModule& Module = ICustomizableObjectModule::Get();
 
-	//Module.UnregisterExternalOperation(FMeshIdentity::StaticStruct());
-	//Module.UnregisterExternalOperation(FMeshNoise::StaticStruct());
+	Module.UnregisterExternalOperation(FMeshIdentity::StaticStruct());
+	Module.UnregisterExternalOperation(FMeshNoise::StaticStruct());
 }
diff --git a/Engine/Plugins/Mutable/Source/CustomizableObject/Internal/MuCO/CustomizableObjectPrivate.h b/Engine/Plugins/Mutable/Source/CustomizableObject/Internal/MuCO/CustomizableObjectPrivate.h
--- a/Engine/Plugins/Mutable/Source/CustomizableObject/Internal/MuCO/CustomizableObjectPrivate.h
+++ b/Engine/Plugins/Mutable/Source/CustomizableObject/Internal/MuCO/CustomizableObjectPrivate.h
@@ -1773,6 +1773,9 @@ class UCustomizableObjectPrivate : public UObject
 #endif
 
 public:
+	// UObject interface
+	virtual void PostLoad() override;
+
 	/** Must be called after unlocking the CustomizableObject. */
 	CUSTOMIZABLEOBJECT_API void SetModel(const TSharedPtr<UE::Mutable::Private::FModel>& Model, const FGuid Identifier, bool bIsCooking);
 	CUSTOMIZABLEOBJECT_API const TSharedPtr<UE::Mutable::Private::FModel>& GetModel();
@@ -1950,6 +1953,8 @@ public:
 	TSharedRef<UE::Mutable::Private::FImageIdRegistry> ImageIdRegistry = MakeShared<UE::Mutable::Private::FImageIdRegistry>();
 
 	TSharedRef<UE::Mutable::Private::FMaterialIdRegistry> MaterialIdRegistry = MakeShared<UE::Mutable::Private::FMaterialIdRegistry>();
+
+	TSharedPtr<UE::Mutable::Private::FModel> SerializedModel;
 	
 	// See UCustomizableObjectSystem::LockObject. Must only be modified from the game thread
 	bool bLocked = false;
diff --git a/Engine/Plugins/Mutable/Source/CustomizableObject/Private/MuCO/CustomizableObject.cpp b/Engine/Plugins/Mutable/Source/CustomizableObject/Private/MuCO/CustomizableObject.cpp
--- a/Engine/Plugins/Mutable/Source/CustomizableObject/Private/MuCO/CustomizableObject.cpp
+++ b/Engine/Plugins/Mutable/Source/CustomizableObject/Private/MuCO/CustomizableObject.cpp
@@ -918,9 +918,7 @@ void UCustomizableObjectPrivate::LoadEmbeddedData(FArchive& Ar)
 		// Load model
 		FUnrealMutableInputStream Stream(Ar);
 		UE::Mutable::Private::FInputArchive Arch(&Stream);
-		TSharedPtr<UE::Mutable::Private::FModel, ESPMode::ThreadSafe> Model = UE::Mutable::Private::FModel::StaticUnserialise(Arch);
-
-		SetModel(Model, FGuid(), false);
+		SerializedModel = UE::Mutable::Private::FModel::StaticUnserialise(Arch);
 	}
 }
 
@@ -2154,6 +2152,17 @@ void UModelStreamableData::PostLoad()
 }
 
 
+void UCustomizableObjectPrivate::PostLoad()
+{
+	Super::PostLoad();
+
+	if (SerializedModel)
+	{
+		SetModel(SerializedModel, FGuid(), false);
+	}
+}
+
+
 void UCustomizableObjectPrivate::SetModel(const TSharedPtr<UE::Mutable::Private::FModel>& Model, const FGuid Id, bool bIsCooking)
 {
 	if (MutableModel == Model
@@ -2178,8 +2187,6 @@ void UCustomizableObjectPrivate::SetModel(const TSharedPtr<UE::Mutable::Private:
 
 	if (Model)
 	{
-		// TODO GMT Extensions disabled due to a crash in packaged build
-		/*
 		UModelResources* LocalModelResources = nullptr;
 #if WITH_EDITORONLY_DATA
 		LocalModelResources = GetModelResources(bIsCooking);
@@ -2189,7 +2196,6 @@ void UCustomizableObjectPrivate::SetModel(const TSharedPtr<UE::Mutable::Private:
 
 		check(LocalModelResources);
 		MutableModel->GetPrivate()->Program.ExternalOperationProvider = MakeShared<UE::Mutable::Private::FExternalOperationProvider>(*LocalModelResources);
-		*/
 	}
 	
 	// Create parameter properties
diff --git a/Engine/Plugins/Mutable/Source/CustomizableObject/Private/MuCO/CustomizableObjectModule.cpp b/Engine/Plugins/Mutable/Source/CustomizableObject/Private/MuCO/CustomizableObjectModule.cpp
--- a/Engine/Plugins/Mutable/Source/CustomizableObject/Private/MuCO/CustomizableObjectModule.cpp
+++ b/Engine/Plugins/Mutable/Source/CustomizableObject/Private/MuCO/CustomizableObjectModule.cpp
@@ -210,8 +210,15 @@ void FCustomizableObjectModule::UnregisterExternalOperation(const UScriptStruct*
 	{
 		return;
 	}
-	
-	RegisteredOperations.Remove(TStrongObjectPtr(Operation));
+
+	for (TArray<TStrongObjectPtr<const UScriptStruct>>::TIterator It = RegisteredOperations.CreateIterator(); It; ++It)
+	{
+		if (It->Get() == Operation)
+		{
+			It.RemoveCurrent();
+			break;
+		}
+	}
 }
 
 
